build script, retry unloading the app domain, then force taskkill automationtoollauncher

diff --git a/Engine/Source/Programs/AutomationToolLauncher/Launcher.cs b/Engine/Source/Programs/AutomationToolLauncher/Launcher.cs
--- a/Engine/Source/Programs/AutomationToolLauncher/Launcher.cs
+++ b/Engine/Source/Programs/AutomationToolLauncher/Launcher.cs
@@ -5,11 +5,24 @@ using System.Linq;
 using System.Text;
 using System.Reflection;
 using System.IO;
+using System.Diagnostics;
 
 namespace AutomationToolLauncher
 {
 	class Launcher
 	{
+        static void KillMe()
+        {
+            Process NewProcess = new Process();
+            NewProcess.StartInfo.FileName = "TaskKill";
+            NewProcess.StartInfo.Arguments = "/IM AutomationToolLauncher.exe /T /F";
+            NewProcess.StartInfo.UseShellExecute = true;
+            Console.WriteLine("Killing with: {0} {1}", NewProcess.StartInfo.FileName, NewProcess.StartInfo.Arguments);
+            NewProcess.Start();
+            System.Threading.Thread.Sleep(25000);
+            Console.WriteLine("We should be dead.");
+        }
+
 		static int Main()
 		{
 			// Create application domain setup information.
@@ -24,17 +37,36 @@ namespace AutomationToolLauncher
 			var UATExecutable = Path.Combine(Domaininfo.ApplicationBase, "AutomationTool.exe");
 			// Default exit code in case UAT does not even start, otherwise we always return UAT's exit code.
 			var ExitCode = 193;
-
-			try
-			{
-				ExitCode = Domain.ExecuteAssembly(UATExecutable, CommandLine);
-				// Unload the application domain.
-				AppDomain.Unload(Domain);
-			}
-			catch (Exception Ex)
-			{
-				Console.WriteLine(Ex.Message);
-			}
+            try
+            {
+                ExitCode = Domain.ExecuteAssembly(UATExecutable, CommandLine);
+            }
+            catch (Exception Ex)
+            {
+                Console.WriteLine("AutomationToolLauncher failed to execute the assembly");
+                Console.WriteLine(Ex.Message);
+            }
+            int Retry = 5;
+            bool Unloaded = false;
+            while (!Unloaded)
+            {
+                try
+                {
+                    // Unload the application domain.
+                    AppDomain.Unload(Domain);
+                    Unloaded = true;
+                }
+                catch (Exception Ex)
+                {
+                    Console.WriteLine(Ex.Message);
+                    if (--Retry == 0)
+                    {
+                        KillMe();
+                    }
+                    Console.WriteLine("AutomationToolLauncher failed unload the app domain, retrying in 30s...");
+                    System.Threading.Thread.Sleep(30000);
+                }
+            }
 			Console.WriteLine("AutomationToolLauncher exiting with ExitCode={0}", ExitCode);
 			return ExitCode;
 		}
