UAT: Reverting recent changes that skipped AutomationToolLauncher

diff --git a/Engine/Build/BatchFiles/RunUAT.bat b/Engine/Build/BatchFiles/RunUAT.bat
--- a/Engine/Build/BatchFiles/RunUAT.bat
+++ b/Engine/Build/BatchFiles/RunUAT.bat
@@ -1,4 +1,4 @@
-rem secho off
+@echo off
 setlocal 
 
 echo Running AutomationTool...
@@ -13,18 +13,6 @@ set UATExecutable=AutomationToolLauncher.exe
 set UATDirectory=Binaries\DotNET\
 set UATNoCompileArg=
 
-rem ## when running automated test, there's no need to use AutomationToolLauncher (fix for EC hangs)
-set UATCmdLineEscaped=%*
-set UATCmdLineEscaped=%UATCmdLineEscaped: =_%
-set UATCmdLineEscaped=%UATCmdLineEscaped:"=_%
-set UATCmdLine=%UATCmdLineEscaped%
-set UATCmdLine=%UATCmdLine: =_%
-set UATCmdLine=%UATCmdLine:EditorTest=%
-set UATCmdLine=%UATCmdLine:GameTest=%
-set UATCmdLine=%UATCmdLine:EditorAutomationTest=%
-set UATCmdLine=%UATCmdLine:GameAutomationTest=%
-if not "%UATCmdLineEscaped%"=="%UATCmdLine%" set UATExecutable=AutomationTool.exe
-
 rem ## Change the CWD to /Engine. 
 pushd %~dp0..\..\
 if not exist Build\BatchFiles\RunUAT.bat goto Error_BatchFileInWrongLocation
diff --git a/Engine/Source/Programs/AutomationTool/TempStorage.cs b/Engine/Source/Programs/AutomationTool/TempStorage.cs
--- a/Engine/Source/Programs/AutomationTool/TempStorage.cs
+++ b/Engine/Source/Programs/AutomationTool/TempStorage.cs
@@ -821,28 +821,12 @@ namespace AutomationTool
                 }
                 var RelativeFile = Filename.Substring(BlockPath.Length);
                 var DestFile = CombinePaths(BaseFolder, RelativeFile);
-                bool bDoCopy = true;
                 if (FileExists_NoExceptions(true, DestFile))
                 {
                     Log("Dest file {0} already exists, deleting and overwriting", DestFile);
-                    DeleteFile_NoExceptions(DestFile);
-                    if (FileExists_NoExceptions(true, DestFile))
-                    {
-                        bool bOkToBeDifferent = DestFile.Contains("AutomationTool.exe");
-                        bOkToBeDifferent = bOkToBeDifferent || DestFile.Contains("\\DotNET\\AutomationScripts\\") || DestFile.Contains("/DotNET/AutomationScripts/");
-						bOkToBeDifferent = bOkToBeDifferent || DestFile.Contains("UnrealBuildTool.exe");
-                        if (bOkToBeDifferent)
-                        {
-                            bDoCopy = false;
-                        }
-                        System.Diagnostics.TraceEventType LogType = bOkToBeDifferent ? System.Diagnostics.TraceEventType.Warning : System.Diagnostics.TraceEventType.Error;
-                        Log(LogType, "Dest file {0} already exists and couldn't be deleted.", DestFile);
-                    }
-                }
-                if (bDoCopy)
-                {
-                    CopyFile(Filename, DestFile, true);
+                    DeleteFile(DestFile);
                 }
+                CopyFile(Filename, DestFile, true);
 
                 Robust_FileExists_NoExceptions(true, DestFile, "Could not copy to {0}");
 
diff --git a/Engine/Source/Programs/AutomationToolLauncher/Launcher.cs b/Engine/Source/Programs/AutomationToolLauncher/Launcher.cs
--- a/Engine/Source/Programs/AutomationToolLauncher/Launcher.cs
+++ b/Engine/Source/Programs/AutomationToolLauncher/Launcher.cs
@@ -5,24 +5,11 @@ using System.Linq;
 using System.Text;
 using System.Reflection;
 using System.IO;
-using System.Diagnostics;
 
 namespace AutomationToolLauncher
 {
 	class Launcher
 	{
-        static void KillMe()
-        {
-            Process NewProcess = new Process();
-            NewProcess.StartInfo.FileName = "TaskKill";
-            NewProcess.StartInfo.Arguments = "/IM AutomationToolLauncher.exe /T /F";
-            NewProcess.StartInfo.UseShellExecute = true;
-            Console.WriteLine("Killing with: {0} {1}", NewProcess.StartInfo.FileName, NewProcess.StartInfo.Arguments);
-            NewProcess.Start();
-            System.Threading.Thread.Sleep(25000);
-            Console.WriteLine("We should be dead.");
-        }
-
 		static int Main()
 		{
 			// Create application domain setup information.
@@ -37,37 +24,18 @@ namespace AutomationToolLauncher
 			var UATExecutable = Path.Combine(Domaininfo.ApplicationBase, "AutomationTool.exe");
 			// Default exit code in case UAT does not even start, otherwise we always return UAT's exit code.
 			var ExitCode = 193;
-            try
-            {
-                ExitCode = Domain.ExecuteAssembly(UATExecutable, CommandLine);
-            }
-            catch (Exception Ex)
-            {
-                Console.WriteLine("AutomationToolLauncher failed to execute the assembly");
-                Console.WriteLine(Ex.Message);
-								Console.WriteLine(Ex.StackTrace);
-            }
-            int Retry = 5;
-            bool Unloaded = false;
-            while (!Unloaded)
-            {
-                try
-                {
-                    // Unload the application domain.
-                    AppDomain.Unload(Domain);
-                    Unloaded = true;
-                }
-                catch (Exception Ex)
-                {
-                    Console.WriteLine(Ex.Message);
-                    if (--Retry == 0)
-                    {
-                        KillMe();
-                    }
-                    Console.WriteLine("AutomationToolLauncher failed unload the app domain, retrying in 30s...");
-                    System.Threading.Thread.Sleep(30000);
-                }
-            }
+
+			try
+			{
+				ExitCode = Domain.ExecuteAssembly(UATExecutable, CommandLine);
+				// Unload the application domain.
+				AppDomain.Unload(Domain);
+			}
+			catch (Exception Ex)
+			{
+				Console.WriteLine(Ex.Message);
+				Console.WriteLine(Ex.StackTrace);
+			}
 			Console.WriteLine("AutomationToolLauncher exiting with ExitCode={0}", ExitCode);
 			return ExitCode;
 		}
