Git PR 371 - AutomationToolLauncher crash fix: needs loadFromRemoteSources

Integrated.

diff --git a/Engine/Source/Programs/AutomationToolLauncher/AutomationToolLauncher.csproj b/Engine/Source/Programs/AutomationToolLauncher/AutomationToolLauncher.csproj
--- a/Engine/Source/Programs/AutomationToolLauncher/AutomationToolLauncher.csproj
+++ b/Engine/Source/Programs/AutomationToolLauncher/AutomationToolLauncher.csproj
@@ -50,6 +50,9 @@
     <Compile Include="Launcher.cs" />
     <Compile Include="Properties\AssemblyInfo.cs" />
   </ItemGroup>
+  <ItemGroup>
+    <None Include="app.config" />
+  </ItemGroup>
   <Import Project="$(MSBuildToolsPath)\Microsoft.CSharp.targets" />
   <!-- To modify your build process, add your task inside one of the targets below and uncomment it. 
        Other similar extension points exist, see Microsoft.Common.targets.
diff --git a/Engine/Source/Programs/AutomationToolLauncher/app.config b/Engine/Source/Programs/AutomationToolLauncher/app.config
new file mode 100644
--- /dev/null
+++ b/Engine/Source/Programs/AutomationToolLauncher/app.config
@@ -0,0 +1,6 @@
+<?xml version="1.0"?>
+<configuration>
+   <runtime>
+      <loadFromRemoteSources enabled="true"/>
+   </runtime>
+</configuration>
