Fix command line arguments passed through the project launcher not being escaped properly. Also remove the custom command line parser from UAT, which doesn't parse standard escape sequences correctly. This code was originally added to allow newlines to be parsed as argument separators, but the maintainance cost of it isn't worth the trouble. I'd rather find and fix any places where that happens.

diff --git a/Engine/Source/Developer/LauncherServices/Private/Launcher/LauncherWorker.cpp b/Engine/Source/Developer/LauncherServices/Private/Launcher/LauncherWorker.cpp
--- a/Engine/Source/Developer/LauncherServices/Private/Launcher/LauncherWorker.cpp
+++ b/Engine/Source/Developer/LauncherServices/Private/Launcher/LauncherWorker.cpp
@@ -536,12 +536,21 @@ FString FLauncherWorker::CreateUATCommand( const ILauncherProfileRef& InProfile,
 				UATCommand += TEXT(" -encryptinifiles");
 			}
 
-			FString additionalOptions = InProfile->GetCookOptions();
-			if (!additionalOptions.IsEmpty())
+			FString AdditionalOptions = InProfile->GetCookOptions();
+			if (!AdditionalOptions.IsEmpty())
 			{
 				UATCommand += TEXT(" -additionalcookeroptions=\"");
-				UATCommand += additionalOptions;
-                UATCommand += "\"";
+
+				// Escape any quotes in the argument list
+				UATCommand += AdditionalOptions.Replace(TEXT("\""), TEXT("\\\""));
+
+				// If the additional options ends with a slash, make sure we don't escape the quote
+				if (UATCommand.EndsWith("\\"))
+				{
+					UATCommand += TEXT("\\");
+				}
+
+				UATCommand += TEXT("\"");
 			}
 
 			if (FParse::Param(FCommandLine::Get(), TEXT("fastcook")))
diff --git a/Engine/Source/Programs/AutomationTool/Program.cs b/Engine/Source/Programs/AutomationTool/Program.cs
--- a/Engine/Source/Programs/AutomationTool/Program.cs
+++ b/Engine/Source/Programs/AutomationTool/Program.cs
@@ -23,11 +23,8 @@ namespace AutomationTool
 		static ProcessManager.CtrlHandlerDelegate CtrlHandlerDelegateInstance = CtrlHandler;
 
 		[STAThread]
-		public static int Main()
+		public static int Main(string[] Arguments)
 		{
-			// Parse the command line
-			string[] Arguments = SharedUtils.ParseCommandLineAndRemoveExe(Environment.CommandLine);
-
             // Ensure UTF8Output flag is respected, since we are initializing logging early in the program.
             if (CommandUtils.ParseParam(Arguments, "-Utf8output"))
             {
diff --git a/Engine/Source/Programs/AutomationToolLauncher/AutomationToolLauncher.csproj b/Engine/Source/Programs/AutomationToolLauncher/AutomationToolLauncher.csproj
--- a/Engine/Source/Programs/AutomationToolLauncher/AutomationToolLauncher.csproj
+++ b/Engine/Source/Programs/AutomationToolLauncher/AutomationToolLauncher.csproj
@@ -46,9 +46,6 @@
     <Reference Include="System.Xml" />
   </ItemGroup>
   <ItemGroup>
-    <Compile Include="..\AutomationTool\AutomationUtils\SharedUtils.cs">
-      <Link>SharedUtils.cs</Link>
-    </Compile>
     <Compile Include="..\DotNETCommon\MetaData.cs">
       <Link>Properties\MetaData.cs</Link>
     </Compile>
@@ -66,4 +63,4 @@
   <Target Name="AfterBuild">
   </Target>
   -->
-</Project>
+</Project>
\ No newline at end of file
diff --git a/Engine/Source/Programs/AutomationToolLauncher/Launcher.cs b/Engine/Source/Programs/AutomationToolLauncher/Launcher.cs
--- a/Engine/Source/Programs/AutomationToolLauncher/Launcher.cs
+++ b/Engine/Source/Programs/AutomationToolLauncher/Launcher.cs
@@ -1,4 +1,4 @@
-﻿// Copyright 1998-2019 Epic Games, Inc. All Rights Reserved.
+// Copyright 1998-2019 Epic Games, Inc. All Rights Reserved.
 using System;
 using System.Collections.Generic;
 using System.Linq;
@@ -10,7 +10,7 @@ namespace AutomationToolLauncher
 {
 	class Launcher
 	{
-		static int Main()
+		static int Main(string[] Arguments)
 		{
 			// Create application domain setup information.
 			var Domaininfo = new AppDomainSetup();
@@ -20,14 +20,13 @@ namespace AutomationToolLauncher
 			// Create the application domain.			
 			var Domain = AppDomain.CreateDomain("AutomationTool", AppDomain.CurrentDomain.Evidence, Domaininfo);
 			// Execute assembly and pass through command line
-			var CommandLine = AutomationTool.SharedUtils.ParseCommandLineAndRemoveExe(Environment.CommandLine);
 			var UATExecutable = Path.Combine(Domaininfo.ApplicationBase, "AutomationTool.exe");
 			// Default exit code in case UAT does not even start, otherwise we always return UAT's exit code.
 			var ExitCode = 193;
 
 			try
 			{
-				ExitCode = Domain.ExecuteAssembly(UATExecutable, CommandLine);
+				ExitCode = Domain.ExecuteAssembly(UATExecutable, Arguments);
 				// Unload the application domain.
 				AppDomain.Unload(Domain);
 			}
