Modernize AutomationToolLauncher
This tool may be unneeded and could be deleted at a later date

#rnx
#Synced-CL 52421477

diff --git a/Engine/Source/Programs/AutomationToolLauncher/AutomationToolLauncher.csproj b/Engine/Source/Programs/AutomationToolLauncher/AutomationToolLauncher.csproj
--- a/Engine/Source/Programs/AutomationToolLauncher/AutomationToolLauncher.csproj
+++ b/Engine/Source/Programs/AutomationToolLauncher/AutomationToolLauncher.csproj
@@ -2,43 +2,45 @@
   <Import Project="..\Shared\UnrealEngine.csproj.props" />
 
   <PropertyGroup>
-    <TargetFramework>net8.0</TargetFramework>
-    <Configuration Condition=" '$(Configuration)' == '' ">Development</Configuration>
     <OutputType>Exe</OutputType>
     <AppendTargetFrameworkToOutputPath>false</AppendTargetFrameworkToOutputPath>
     <GenerateAssemblyInfo>false</GenerateAssemblyInfo>
     <GenerateTargetFrameworkAttribute>false</GenerateTargetFrameworkAttribute>
-    <Configurations>Debug;Release;Development;Analyze</Configurations>
     <ApplicationManifest>app.manifest</ApplicationManifest>
     <RootNamespace>AutomationToolLauncher</RootNamespace>
     <AssemblyName>AutomationToolLauncher</AssemblyName>
-    <WarningsNotAsErrors>612,618</WarningsNotAsErrors>
     <OutputPath>..\..\..\Binaries\DotNET\AutomationToolLauncher</OutputPath>
   </PropertyGroup>
 
-  <PropertyGroup Condition=" '$(Configuration)|$(Platform)' == 'Development|AnyCPU' ">
-    <DefineConstants>$(DefineConstants);TRACE</DefineConstants>
-    <Optimize>true</Optimize>
-  </PropertyGroup>
-  <PropertyGroup Condition=" '$(Configuration)|$(Platform)' == 'Release|AnyCPU' ">
-    <DefineConstants>$(DefineConstants);TRACE</DefineConstants>
-    <Optimize>true</Optimize>
-  </PropertyGroup>
-  <PropertyGroup Condition=" '$(Configuration)|$(Platform)' == 'Debug|AnyCPU' ">
-    <DefineConstants>$(DefineConstants);DEBUG;TRACE</DefineConstants>
-    <DebugSymbols>true</DebugSymbols>
-    <DebugType>full</DebugType>
-  </PropertyGroup>
-  <PropertyGroup Condition=" '$(Configuration)|$(Platform)' == 'Analyze|AnyCPU' ">
-    <RunAnalyzersDuringBuild>True</RunAnalyzersDuringBuild>
-    <TreatWarningsAsErrors>false</TreatWarningsAsErrors>
-    <GenerateDocumentationFile>True</GenerateDocumentationFile>
-    <DocumentationFile></DocumentationFile>
-  </PropertyGroup>
-
   <ItemGroup>
     <Compile Include="..\Shared\MetaData.cs">
       <Link>Properties\MetaData.cs</Link>
     </Compile>
   </ItemGroup>
+
+  <ItemGroup>
+    <ProjectReference Include="..\Shared\EpicGames.Build\EpicGames.Build.csproj" />
+    <ProjectReference Include="..\Shared\EpicGames.Core\EpicGames.Core.csproj" />
+  </ItemGroup>
+
+  <!-- Strip transitive dependencies not reachable from this launcher's code paths -->
+  <Target Name="RemoveUnneededTransitiveDependencies" AfterTargets="ResolveReferences">
+    <ItemGroup>
+      <!-- EpicGames.IoHash and its native Blake3 binaries (pulled in by EpicGames.Build) -->
+      <ReferenceCopyLocalPaths Remove="@(ReferenceCopyLocalPaths)" Condition="'%(Filename)' == 'EpicGames.IoHash'" />
+      <ReferenceCopyLocalPaths Remove="@(ReferenceCopyLocalPaths)" Condition="'%(Filename)' == 'Blake3'" />
+      <!-- EpicGames.MsBuild and Microsoft.Build.Locator (pulled in by EpicGames.Build) -->
+      <ReferenceCopyLocalPaths Remove="@(ReferenceCopyLocalPaths)" Condition="'%(Filename)' == 'EpicGames.MsBuild'" />
+      <ReferenceCopyLocalPaths Remove="@(ReferenceCopyLocalPaths)" Condition="'%(Filename)' == 'Microsoft.Build.Locator'" />
+      <!-- JetBrains.Annotations, OpenTracing and Polly (pulled in by EpicGames.Core) -->
+      <ReferenceCopyLocalPaths Remove="@(ReferenceCopyLocalPaths)" Condition="'%(Filename)' == 'JetBrains.Annotations'" />
+      <ReferenceCopyLocalPaths Remove="@(ReferenceCopyLocalPaths)" Condition="'%(Filename)' == 'OpenTracing'" />
+      <ReferenceCopyLocalPaths Remove="@(ReferenceCopyLocalPaths)" Condition="'%(Filename)' == 'Polly'" />
+      <ReferenceCopyLocalPaths Remove="@(ReferenceCopyLocalPaths)" Condition="'%(Filename)' == 'Polly.Extensions.Http'" />
+      <!-- Microsoft.Extensions.* (pulled in by EpicGames.Core) -->
+      <ReferenceCopyLocalPaths Remove="@(ReferenceCopyLocalPaths)" Condition="$([System.String]::Copy('%(Filename)').StartsWith('Microsoft.Extensions.'))" />
+      <!-- Native runtime folders (blake3_dotnet.dll etc.) -->
+      <ReferenceCopyLocalPaths Remove="@(ReferenceCopyLocalPaths)" Condition="$([System.String]::Copy('%(Identity)').Contains('\runtimes\'))" />
+    </ItemGroup>
+  </Target>
 </Project>
\ No newline at end of file
diff --git a/Engine/Source/Programs/AutomationToolLauncher/Launcher.cs b/Engine/Source/Programs/AutomationToolLauncher/Launcher.cs
--- a/Engine/Source/Programs/AutomationToolLauncher/Launcher.cs
+++ b/Engine/Source/Programs/AutomationToolLauncher/Launcher.cs
@@ -1,51 +1,48 @@
 // Copyright Epic Games, Inc. All Rights Reserved.
+using EpicGames.Core;
 using System;
-using System.Reflection;
-using System.IO;
-using System.Linq;
 using System.Diagnostics;
-using System.Dynamic;
+using System.Threading.Tasks;
+using UnrealBuildBase;
 
-namespace AutomationToolLauncher
+namespace AutomationToolLauncher;
+
+class Launcher
 {
-	class Launcher
+	static async Task<int> Main(string[] arguments)
 	{
-		static int Main(string[] Arguments)
-		{
-			return Run(Arguments);
+		FileReference uatPath = GetUatPath();
 
-		}
-		static int Run(string[] Arguments)
+		if (!FileReference.Exists(uatPath))
 		{
+			Console.WriteLine($"AutomationTool does not exist at: {uatPath}");
+			return -1;
+		}
 
-			string ApplicationBase = Path.GetDirectoryName(Assembly.GetExecutingAssembly().Location);
-			string UATExecutable = Path.Combine(ApplicationBase, "..\\AutomationTool", "AutomationTool.exe");
-
-			if (!File.Exists(UATExecutable))
-			{
-				Console.WriteLine(string.Format("AutomationTool does not exist at: {0}", UATExecutable));
-				return -1;
-			}
-
-			try
-			{
-				ProcessStartInfo StartInfo = new ProcessStartInfo(UATExecutable);
-				foreach (string s in Arguments)
-				{
-					StartInfo.ArgumentList.Add(s);
-				}
-				Process uatProcess = Process.Start(StartInfo);
-				uatProcess.WaitForExit();
-				Environment.Exit(uatProcess.ExitCode);
-			}
-			catch (Exception Ex)
+		try
+		{
+			ProcessStartInfo psi = new(Unreal.DotnetPath.FullName);
+			psi.ArgumentList.Add(uatPath.FullName);
+			foreach (string arg in arguments)
 			{
-				Console.WriteLine(Ex.Message);
-				Console.WriteLine(Ex.StackTrace);
+				psi.ArgumentList.Add(arg);
 			}
+			using Process uatProcess = Process.Start(psi) ?? throw new InvalidOperationException("Failed to start AutomationTool process.");
+			await uatProcess.WaitForExitAsync();
+			return uatProcess.ExitCode;
+		}
+		catch (Exception ex)
+		{
+			Console.WriteLine(ex.Message);
+			Console.WriteLine(ex.StackTrace);
+		}
 
-			return -1;
+		return -1;
+	}
 
-		}
+	static FileReference GetUatPath()
+	{
+		DirectoryReference applicationBase = new DirectoryReference(AppContext.BaseDirectory);
+		return FileReference.Combine(applicationBase, "..", "AutomationTool", "AutomationTool.dll");
 	}
 }
diff --git a/Engine/Source/Programs/AutomationToolLauncher/Properties/AssemblyInfo.cs b/Engine/Source/Programs/AutomationToolLauncher/Properties/AssemblyInfo.cs
--- a/Engine/Source/Programs/AutomationToolLauncher/Properties/AssemblyInfo.cs
+++ b/Engine/Source/Programs/AutomationToolLauncher/Properties/AssemblyInfo.cs
@@ -1,12 +1,5 @@
 ﻿// Copyright Epic Games, Inc. All Rights Reserved.
 
 using System.Reflection;
-using System.Runtime.CompilerServices;
-using System.Runtime.InteropServices;
 
-// General Information about an assembly is controlled through the following 
-// set of attributes. Change these attribute values to modify the information
-// associated with an assembly.
-[assembly: AssemblyTitle("UnrealAutomationTool")]
-[assembly: AssemblyDescription("")]
-[assembly: AssemblyConfiguration("")]
+[assembly: AssemblyTitle("AutomationToolLauncher")]
